Some devices require taps that are unreasonably small and precise in order not to be read as swipes. Adding a small dead zone that scales with resolution to compensate.

diff --git a/Templates/FP_FirstPerson/Source/FP_FirstPerson/FP_FirstPersonCharacter.cpp b/Templates/FP_FirstPerson/Source/FP_FirstPerson/FP_FirstPersonCharacter.cpp
--- a/Templates/FP_FirstPerson/Source/FP_FirstPerson/FP_FirstPersonCharacter.cpp
+++ b/Templates/FP_FirstPerson/Source/FP_FirstPerson/FP_FirstPersonCharacter.cpp
@@ -187,13 +187,13 @@ void AFP_FirstPersonCharacter::TouchUpdate(const ETouchIndex::Type FingerIndex,
 				FVector2D ScreenSize;
 				ViewportClient->GetViewportSize(ScreenSize);
 				FVector2D ScaledDelta = FVector2D(MoveDelta.X, MoveDelta.Y) / ScreenSize;
-				if (ScaledDelta.X != 0.0f)
+				if (FMath::Abs(ScaledDelta.X) >= (4.0f / ScreenSize.X))
 				{
 					TouchItem.bMoved = true;
 					float Value = ScaledDelta.X * BaseTurnRate;
 					AddControllerYawInput(Value);
 				}
-				if (ScaledDelta.Y != 0.0f)
+				if (FMath::Abs(ScaledDelta.Y) >= (4.0f / ScreenSize.Y))
 				{
 					TouchItem.bMoved = true;
 					float Value = ScaledDelta.Y* BaseTurnRate;
