Sample/template typo fixes in (local) variable names.

diff --git a/Templates/FP_FirstPerson/Source/FP_FirstPerson/FP_FirstPersonHUD.cpp b/Templates/FP_FirstPerson/Source/FP_FirstPerson/FP_FirstPersonHUD.cpp
--- a/Templates/FP_FirstPerson/Source/FP_FirstPerson/FP_FirstPersonHUD.cpp
+++ b/Templates/FP_FirstPerson/Source/FP_FirstPerson/FP_FirstPersonHUD.cpp
@@ -9,8 +9,8 @@
 AFP_FirstPersonHUD::AFP_FirstPersonHUD()
 {
 	// Set the crosshair texture
-	static ConstructorHelpers::FObjectFinder<UTexture2D> CrosshiarTexObj(TEXT("/Game/FirstPerson/Textures/FirstPersonCrosshair"));
-	CrosshairTex = CrosshiarTexObj.Object;
+	static ConstructorHelpers::FObjectFinder<UTexture2D> CrosshairTexObj(TEXT("/Game/FirstPerson/Textures/FirstPersonCrosshair"));
+	CrosshairTex = CrosshairTexObj.Object;
 }
 
 
@@ -32,4 +32,3 @@ void AFP_FirstPersonHUD::DrawHUD()
 	TileItem.BlendMode = SE_BLEND_Translucent;
 	Canvas->DrawItem( TileItem );
 }
-
diff --git a/Templates/TP_FirstPerson/Source/TP_FirstPerson/TP_FirstPersonHUD.cpp b/Templates/TP_FirstPerson/Source/TP_FirstPerson/TP_FirstPersonHUD.cpp
--- a/Templates/TP_FirstPerson/Source/TP_FirstPerson/TP_FirstPersonHUD.cpp
+++ b/Templates/TP_FirstPerson/Source/TP_FirstPerson/TP_FirstPersonHUD.cpp
@@ -9,8 +9,8 @@
 ATP_FirstPersonHUD::ATP_FirstPersonHUD()
 {
 	// Set the crosshair texture
-	static ConstructorHelpers::FObjectFinder<UTexture2D> CrosshiarTexObj(TEXT("/Game/FirstPerson/Textures/FirstPersonCrosshair"));
-	CrosshairTex = CrosshiarTexObj.Object;
+	static ConstructorHelpers::FObjectFinder<UTexture2D> CrosshairTexObj(TEXT("/Game/FirstPerson/Textures/FirstPersonCrosshair"));
+	CrosshairTex = CrosshairTexObj.Object;
 }
 
 
@@ -32,4 +32,3 @@ void ATP_FirstPersonHUD::DrawHUD()
 	TileItem.BlendMode = SE_BLEND_Translucent;
 	Canvas->DrawItem( TileItem );
 }
-
