Merging //UE4/Dev-Main to Dev-VR (//UE4/Dev-VR)
CL 2 of 8

diff --git a/GenerateProjectFiles.command b/GenerateProjectFiles.command
--- a/GenerateProjectFiles.command
+++ b/GenerateProjectFiles.command
@@ -1,5 +1,5 @@
 #!/bin/sh
-# Copyright 1998-2018 Epic Games, Inc. All Rights Reserved.
+# Copyright 1998-2019 Epic Games, Inc. All Rights Reserved.
 #
 # Simple wrapper around GenerateProjectFiles.sh using the
 # .command extension enables it to be run from the OSX Finder.
diff --git a/GenerateProjectFiles.sh b/GenerateProjectFiles.sh
--- a/GenerateProjectFiles.sh
+++ b/GenerateProjectFiles.sh
@@ -1,5 +1,5 @@
 #!/bin/sh
-# Copyright 1998-2018 Epic Games, Inc. All Rights Reserved.
+# Copyright 1998-2019 Epic Games, Inc. All Rights Reserved.
 
 set -e
 
diff --git a/Templates/FP_FirstPerson/Source/FP_FirstPerson.Target.cs b/Templates/FP_FirstPerson/Source/FP_FirstPerson.Target.cs
--- a/Templates/FP_FirstPerson/Source/FP_FirstPerson.Target.cs
+++ b/Templates/FP_FirstPerson/Source/FP_FirstPerson.Target.cs
@@ -1,4 +1,4 @@
-// Copyright 1998-2018 Epic Games, Inc. All Rights Reserved.
+// Copyright 1998-2019 Epic Games, Inc. All Rights Reserved.
 
 using UnrealBuildTool;
 using System.Collections.Generic;
diff --git a/Templates/FP_FirstPerson/Source/FP_FirstPerson/FP_FirstPerson.Build.cs b/Templates/FP_FirstPerson/Source/FP_FirstPerson/FP_FirstPerson.Build.cs
--- a/Templates/FP_FirstPerson/Source/FP_FirstPerson/FP_FirstPerson.Build.cs
+++ b/Templates/FP_FirstPerson/Source/FP_FirstPerson/FP_FirstPerson.Build.cs
@@ -1,4 +1,4 @@
-// Copyright 1998-2018 Epic Games, Inc. All Rights Reserved.
+// Copyright 1998-2019 Epic Games, Inc. All Rights Reserved.
 
 using UnrealBuildTool;
 
diff --git a/Templates/FP_FirstPerson/Source/FP_FirstPerson/FP_FirstPerson.cpp b/Templates/FP_FirstPerson/Source/FP_FirstPerson/FP_FirstPerson.cpp
--- a/Templates/FP_FirstPerson/Source/FP_FirstPerson/FP_FirstPerson.cpp
+++ b/Templates/FP_FirstPerson/Source/FP_FirstPerson/FP_FirstPerson.cpp
@@ -1,4 +1,4 @@
-// Copyright 1998-2018 Epic Games, Inc. All Rights Reserved.
+// Copyright 1998-2019 Epic Games, Inc. All Rights Reserved.
 
 #include "FP_FirstPerson.h"
 #include "ModuleManager.h"
diff --git a/Templates/FP_FirstPerson/Source/FP_FirstPerson/FP_FirstPerson.h b/Templates/FP_FirstPerson/Source/FP_FirstPerson/FP_FirstPerson.h
--- a/Templates/FP_FirstPerson/Source/FP_FirstPerson/FP_FirstPerson.h
+++ b/Templates/FP_FirstPerson/Source/FP_FirstPerson/FP_FirstPerson.h
@@ -1,4 +1,4 @@
-// Copyright 1998-2018 Epic Games, Inc. All Rights Reserved.
+// Copyright 1998-2019 Epic Games, Inc. All Rights Reserved.
 
 #pragma once
 
diff --git a/Templates/FP_FirstPerson/Source/FP_FirstPerson/FP_FirstPersonCharacter.cpp b/Templates/FP_FirstPerson/Source/FP_FirstPerson/FP_FirstPersonCharacter.cpp
--- a/Templates/FP_FirstPerson/Source/FP_FirstPerson/FP_FirstPersonCharacter.cpp
+++ b/Templates/FP_FirstPerson/Source/FP_FirstPerson/FP_FirstPersonCharacter.cpp
@@ -1,4 +1,4 @@
-// Copyright 1998-2018 Epic Games, Inc. All Rights Reserved.
+// Copyright 1998-2019 Epic Games, Inc. All Rights Reserved.
 
 #include "FP_FirstPersonCharacter.h"
 #include "Animation/AnimInstance.h"
@@ -247,7 +247,6 @@ FHitResult AFP_FirstPersonCharacter::WeaponTrace(const FVector& StartTrace, cons
 {
 	// Perform trace to retrieve hit info
 	FCollisionQueryParams TraceParams(SCENE_QUERY_STAT(WeaponTrace), true, Instigator);
-	TraceParams.bTraceAsyncScene = true;
 	TraceParams.bReturnPhysicalMaterial = true;
 
 	FHitResult Hit(ForceInit);
diff --git a/Templates/FP_FirstPerson/Source/FP_FirstPerson/FP_FirstPersonCharacter.h b/Templates/FP_FirstPerson/Source/FP_FirstPerson/FP_FirstPersonCharacter.h
--- a/Templates/FP_FirstPerson/Source/FP_FirstPerson/FP_FirstPersonCharacter.h
+++ b/Templates/FP_FirstPerson/Source/FP_FirstPerson/FP_FirstPersonCharacter.h
@@ -1,4 +1,4 @@
-// Copyright 1998-2018 Epic Games, Inc. All Rights Reserved.
+// Copyright 1998-2019 Epic Games, Inc. All Rights Reserved.
 
 #pragma once
 
diff --git a/Templates/FP_FirstPerson/Source/FP_FirstPerson/FP_FirstPersonGameMode.cpp b/Templates/FP_FirstPerson/Source/FP_FirstPerson/FP_FirstPersonGameMode.cpp
--- a/Templates/FP_FirstPerson/Source/FP_FirstPerson/FP_FirstPersonGameMode.cpp
+++ b/Templates/FP_FirstPerson/Source/FP_FirstPerson/FP_FirstPersonGameMode.cpp
@@ -1,4 +1,4 @@
-// Copyright 1998-2018 Epic Games, Inc. All Rights Reserved.
+// Copyright 1998-2019 Epic Games, Inc. All Rights Reserved.
 
 #include "FP_FirstPersonGameMode.h"
 #include "FP_FirstPersonHUD.h"
diff --git a/Templates/FP_FirstPerson/Source/FP_FirstPerson/FP_FirstPersonGameMode.h b/Templates/FP_FirstPerson/Source/FP_FirstPerson/FP_FirstPersonGameMode.h
--- a/Templates/FP_FirstPerson/Source/FP_FirstPerson/FP_FirstPersonGameMode.h
+++ b/Templates/FP_FirstPerson/Source/FP_FirstPerson/FP_FirstPersonGameMode.h
@@ -1,4 +1,4 @@
-// Copyright 1998-2018 Epic Games, Inc. All Rights Reserved.
+// Copyright 1998-2019 Epic Games, Inc. All Rights Reserved.
 
 #pragma once
 
diff --git a/Templates/FP_FirstPerson/Source/FP_FirstPerson/FP_FirstPersonHUD.cpp b/Templates/FP_FirstPerson/Source/FP_FirstPerson/FP_FirstPersonHUD.cpp
--- a/Templates/FP_FirstPerson/Source/FP_FirstPerson/FP_FirstPersonHUD.cpp
+++ b/Templates/FP_FirstPerson/Source/FP_FirstPerson/FP_FirstPersonHUD.cpp
@@ -1,4 +1,4 @@
-// Copyright 1998-2018 Epic Games, Inc. All Rights Reserved.
+// Copyright 1998-2019 Epic Games, Inc. All Rights Reserved.
 
 #include "FP_FirstPersonHUD.h"
 #include "Engine/Canvas.h"
diff --git a/Templates/FP_FirstPerson/Source/FP_FirstPerson/FP_FirstPersonHUD.h b/Templates/FP_FirstPerson/Source/FP_FirstPerson/FP_FirstPersonHUD.h
--- a/Templates/FP_FirstPerson/Source/FP_FirstPerson/FP_FirstPersonHUD.h
+++ b/Templates/FP_FirstPerson/Source/FP_FirstPerson/FP_FirstPersonHUD.h
@@ -1,4 +1,4 @@
-// Copyright 1998-2018 Epic Games, Inc. All Rights Reserved.
+// Copyright 1998-2019 Epic Games, Inc. All Rights Reserved.
 
 #pragma once
 
diff --git a/Templates/FP_FirstPerson/Source/FP_FirstPersonEditor.Target.cs b/Templates/FP_FirstPerson/Source/FP_FirstPersonEditor.Target.cs
--- a/Templates/FP_FirstPerson/Source/FP_FirstPersonEditor.Target.cs
+++ b/Templates/FP_FirstPerson/Source/FP_FirstPersonEditor.Target.cs
@@ -1,4 +1,4 @@
-// Copyright 1998-2018 Epic Games, Inc. All Rights Reserved.
+// Copyright 1998-2019 Epic Games, Inc. All Rights Reserved.
 
 using UnrealBuildTool;
 using System.Collections.Generic;
diff --git a/Templates/TP_2DSideScroller/Source/TP_2DSideScroller.Target.cs b/Templates/TP_2DSideScroller/Source/TP_2DSideScroller.Target.cs
--- a/Templates/TP_2DSideScroller/Source/TP_2DSideScroller.Target.cs
+++ b/Templates/TP_2DSideScroller/Source/TP_2DSideScroller.Target.cs
@@ -1,4 +1,4 @@
-// Copyright 1998-2018 Epic Games, Inc. All Rights Reserved.
+// Copyright 1998-2019 Epic Games, Inc. All Rights Reserved.
 
 using UnrealBuildTool;
 using System.Collections.Generic;
diff --git a/Templates/TP_2DSideScroller/Source/TP_2DSideScroller/TP_2DSideScroller.Build.cs b/Templates/TP_2DSideScroller/Source/TP_2DSideScroller/TP_2DSideScroller.Build.cs
--- a/Templates/TP_2DSideScroller/Source/TP_2DSideScroller/TP_2DSideScroller.Build.cs
+++ b/Templates/TP_2DSideScroller/Source/TP_2DSideScroller/TP_2DSideScroller.Build.cs
@@ -1,4 +1,4 @@
-// Copyright 1998-2018 Epic Games, Inc. All Rights Reserved.
+// Copyright 1998-2019 Epic Games, Inc. All Rights Reserved.
 
 using UnrealBuildTool;
 
diff --git a/Templates/TP_2DSideScroller/Source/TP_2DSideScroller/TP_2DSideScroller.cpp b/Templates/TP_2DSideScroller/Source/TP_2DSideScroller/TP_2DSideScroller.cpp
--- a/Templates/TP_2DSideScroller/Source/TP_2DSideScroller/TP_2DSideScroller.cpp
+++ b/Templates/TP_2DSideScroller/Source/TP_2DSideScroller/TP_2DSideScroller.cpp
@@ -1,4 +1,4 @@
-// Copyright 1998-2018 Epic Games, Inc. All Rights Reserved.
+// Copyright 1998-2019 Epic Games, Inc. All Rights Reserved.
 
 #include "TP_2DSideScroller.h"
 #include "Modules/ModuleManager.h"
diff --git a/Templates/TP_2DSideScroller/Source/TP_2DSideScroller/TP_2DSideScroller.h b/Templates/TP_2DSideScroller/Source/TP_2DSideScroller/TP_2DSideScroller.h
--- a/Templates/TP_2DSideScroller/Source/TP_2DSideScroller/TP_2DSideScroller.h
+++ b/Templates/TP_2DSideScroller/Source/TP_2DSideScroller/TP_2DSideScroller.h
@@ -1,4 +1,4 @@
-// Copyright 1998-2018 Epic Games, Inc. All Rights Reserved.
+// Copyright 1998-2019 Epic Games, Inc. All Rights Reserved.
 
 #pragma once
 
diff --git a/Templates/TP_2DSideScroller/Source/TP_2DSideScroller/TP_2DSideScrollerCharacter.cpp b/Templates/TP_2DSideScroller/Source/TP_2DSideScroller/TP_2DSideScrollerCharacter.cpp
--- a/Templates/TP_2DSideScroller/Source/TP_2DSideScroller/TP_2DSideScrollerCharacter.cpp
+++ b/Templates/TP_2DSideScroller/Source/TP_2DSideScroller/TP_2DSideScrollerCharacter.cpp
@@ -1,4 +1,4 @@
-// Copyright 1998-2018 Epic Games, Inc. All Rights Reserved.
+// Copyright 1998-2019 Epic Games, Inc. All Rights Reserved.
 
 #include "TP_2DSideScrollerCharacter.h"
 #include "PaperFlipbookComponent.h"
diff --git a/Templates/TP_2DSideScroller/Source/TP_2DSideScroller/TP_2DSideScrollerCharacter.h b/Templates/TP_2DSideScroller/Source/TP_2DSideScroller/TP_2DSideScrollerCharacter.h
--- a/Templates/TP_2DSideScroller/Source/TP_2DSideScroller/TP_2DSideScrollerCharacter.h
+++ b/Templates/TP_2DSideScroller/Source/TP_2DSideScroller/TP_2DSideScrollerCharacter.h
@@ -1,4 +1,4 @@
-// Copyright 1998-2018 Epic Games, Inc. All Rights Reserved.
+// Copyright 1998-2019 Epic Games, Inc. All Rights Reserved.
 
 #pragma once
 
diff --git a/Templates/TP_2DSideScroller/Source/TP_2DSideScroller/TP_2DSideScrollerGameMode.cpp b/Templates/TP_2DSideScroller/Source/TP_2DSideScroller/TP_2DSideScrollerGameMode.cpp
--- a/Templates/TP_2DSideScroller/Source/TP_2DSideScroller/TP_2DSideScrollerGameMode.cpp
+++ b/Templates/TP_2DSideScroller/Source/TP_2DSideScroller/TP_2DSideScrollerGameMode.cpp
@@ -1,4 +1,4 @@
-// Copyright 1998-2018 Epic Games, Inc. All Rights Reserved.
+// Copyright 1998-2019 Epic Games, Inc. All Rights Reserved.
 
 #include "TP_2DSideScrollerGameMode.h"
 #include "TP_2DSideScrollerCharacter.h"
diff --git a/Templates/TP_2DSideScroller/Source/TP_2DSideScroller/TP_2DSideScrollerGameMode.h b/Templates/TP_2DSideScroller/Source/TP_2DSideScroller/TP_2DSideScrollerGameMode.h
--- a/Templates/TP_2DSideScroller/Source/TP_2DSideScroller/TP_2DSideScrollerGameMode.h
+++ b/Templates/TP_2DSideScroller/Source/TP_2DSideScroller/TP_2DSideScrollerGameMode.h
@@ -1,4 +1,4 @@
-// Copyright 1998-2018 Epic Games, Inc. All Rights Reserved.
+// Copyright 1998-2019 Epic Games, Inc. All Rights Reserved.
 
 #pragma once
 
diff --git a/Templates/TP_2DSideScroller/Source/TP_2DSideScrollerEditor.Target.cs b/Templates/TP_2DSideScroller/Source/TP_2DSideScrollerEditor.Target.cs
--- a/Templates/TP_2DSideScroller/Source/TP_2DSideScrollerEditor.Target.cs
+++ b/Templates/TP_2DSideScroller/Source/TP_2DSideScrollerEditor.Target.cs
@@ -1,4 +1,4 @@
-// Copyright 1998-2018 Epic Games, Inc. All Rights Reserved.
+// Copyright 1998-2019 Epic Games, Inc. All Rights Reserved.
 
 using UnrealBuildTool;
 using System.Collections.Generic;
diff --git a/Templates/TP_FirstPerson/Source/TP_FirstPerson.Target.cs b/Templates/TP_FirstPerson/Source/TP_FirstPerson.Target.cs
--- a/Templates/TP_FirstPerson/Source/TP_FirstPerson.Target.cs
+++ b/Templates/TP_FirstPerson/Source/TP_FirstPerson.Target.cs
@@ -1,4 +1,4 @@
-// Copyright 1998-2018 Epic Games, Inc. All Rights Reserved.
+// Copyright 1998-2019 Epic Games, Inc. All Rights Reserved.
 
 using UnrealBuildTool;
 using System.Collections.Generic;
diff --git a/Templates/TP_FirstPerson/Source/TP_FirstPerson/TP_FirstPerson.Build.cs b/Templates/TP_FirstPerson/Source/TP_FirstPerson/TP_FirstPerson.Build.cs
--- a/Templates/TP_FirstPerson/Source/TP_FirstPerson/TP_FirstPerson.Build.cs
+++ b/Templates/TP_FirstPerson/Source/TP_FirstPerson/TP_FirstPerson.Build.cs
@@ -1,4 +1,4 @@
-// Copyright 1998-2018 Epic Games, Inc. All Rights Reserved.
+// Copyright 1998-2019 Epic Games, Inc. All Rights Reserved.
 
 using UnrealBuildTool;
 
diff --git a/Templates/TP_FirstPerson/Source/TP_FirstPerson/TP_FirstPerson.cpp b/Templates/TP_FirstPerson/Source/TP_FirstPerson/TP_FirstPerson.cpp
--- a/Templates/TP_FirstPerson/Source/TP_FirstPerson/TP_FirstPerson.cpp
+++ b/Templates/TP_FirstPerson/Source/TP_FirstPerson/TP_FirstPerson.cpp
@@ -1,4 +1,4 @@
-// Copyright 1998-2018 Epic Games, Inc. All Rights Reserved.
+// Copyright 1998-2019 Epic Games, Inc. All Rights Reserved.
 
 #include "TP_FirstPerson.h"
 #include "Modules/ModuleManager.h"
diff --git a/Templates/TP_FirstPerson/Source/TP_FirstPerson/TP_FirstPerson.h b/Templates/TP_FirstPerson/Source/TP_FirstPerson/TP_FirstPerson.h
--- a/Templates/TP_FirstPerson/Source/TP_FirstPerson/TP_FirstPerson.h
+++ b/Templates/TP_FirstPerson/Source/TP_FirstPerson/TP_FirstPerson.h
@@ -1,4 +1,4 @@
-// Copyright 1998-2018 Epic Games, Inc. All Rights Reserved.
+// Copyright 1998-2019 Epic Games, Inc. All Rights Reserved.
 
 #pragma once
 
diff --git a/Templates/TP_FirstPerson/Source/TP_FirstPerson/TP_FirstPersonCharacter.cpp b/Templates/TP_FirstPerson/Source/TP_FirstPerson/TP_FirstPersonCharacter.cpp
--- a/Templates/TP_FirstPerson/Source/TP_FirstPerson/TP_FirstPersonCharacter.cpp
+++ b/Templates/TP_FirstPerson/Source/TP_FirstPerson/TP_FirstPersonCharacter.cpp
@@ -1,4 +1,4 @@
-// Copyright 1998-2018 Epic Games, Inc. All Rights Reserved.
+// Copyright 1998-2019 Epic Games, Inc. All Rights Reserved.
 
 #include "TP_FirstPersonCharacter.h"
 #include "TP_FirstPersonProjectile.h"
diff --git a/Templates/TP_FirstPerson/Source/TP_FirstPerson/TP_FirstPersonCharacter.h b/Templates/TP_FirstPerson/Source/TP_FirstPerson/TP_FirstPersonCharacter.h
--- a/Templates/TP_FirstPerson/Source/TP_FirstPerson/TP_FirstPersonCharacter.h
+++ b/Templates/TP_FirstPerson/Source/TP_FirstPerson/TP_FirstPersonCharacter.h
@@ -1,4 +1,4 @@
-// Copyright 1998-2018 Epic Games, Inc. All Rights Reserved.
+// Copyright 1998-2019 Epic Games, Inc. All Rights Reserved.
 
 #pragma once
 
diff --git a/Templates/TP_FirstPerson/Source/TP_FirstPerson/TP_FirstPersonGameMode.cpp b/Templates/TP_FirstPerson/Source/TP_FirstPerson/TP_FirstPersonGameMode.cpp
--- a/Templates/TP_FirstPerson/Source/TP_FirstPerson/TP_FirstPersonGameMode.cpp
+++ b/Templates/TP_FirstPerson/Source/TP_FirstPerson/TP_FirstPersonGameMode.cpp
@@ -1,4 +1,4 @@
-// Copyright 1998-2018 Epic Games, Inc. All Rights Reserved.
+// Copyright 1998-2019 Epic Games, Inc. All Rights Reserved.
 
 #include "TP_FirstPersonGameMode.h"
 #include "TP_FirstPersonHUD.h"
diff --git a/Templates/TP_FirstPerson/Source/TP_FirstPerson/TP_FirstPersonGameMode.h b/Templates/TP_FirstPerson/Source/TP_FirstPerson/TP_FirstPersonGameMode.h
--- a/Templates/TP_FirstPerson/Source/TP_FirstPerson/TP_FirstPersonGameMode.h
+++ b/Templates/TP_FirstPerson/Source/TP_FirstPerson/TP_FirstPersonGameMode.h
@@ -1,4 +1,4 @@
-// Copyright 1998-2018 Epic Games, Inc. All Rights Reserved.
+// Copyright 1998-2019 Epic Games, Inc. All Rights Reserved.
 
 #pragma once
 
diff --git a/Templates/TP_FirstPerson/Source/TP_FirstPerson/TP_FirstPersonHUD.cpp b/Templates/TP_FirstPerson/Source/TP_FirstPerson/TP_FirstPersonHUD.cpp
--- a/Templates/TP_FirstPerson/Source/TP_FirstPerson/TP_FirstPersonHUD.cpp
+++ b/Templates/TP_FirstPerson/Source/TP_FirstPerson/TP_FirstPersonHUD.cpp
@@ -1,4 +1,4 @@
-// Copyright 1998-2018 Epic Games, Inc. All Rights Reserved.
+// Copyright 1998-2019 Epic Games, Inc. All Rights Reserved.
 
 #include "TP_FirstPersonHUD.h"
 #include "Engine/Canvas.h"
diff --git a/Templates/TP_FirstPerson/Source/TP_FirstPerson/TP_FirstPersonHUD.h b/Templates/TP_FirstPerson/Source/TP_FirstPerson/TP_FirstPersonHUD.h
--- a/Templates/TP_FirstPerson/Source/TP_FirstPerson/TP_FirstPersonHUD.h
+++ b/Templates/TP_FirstPerson/Source/TP_FirstPerson/TP_FirstPersonHUD.h
@@ -1,4 +1,4 @@
-// Copyright 1998-2018 Epic Games, Inc. All Rights Reserved.
+// Copyright 1998-2019 Epic Games, Inc. All Rights Reserved.
 
 #pragma once 
 
diff --git a/Templates/TP_FirstPerson/Source/TP_FirstPerson/TP_FirstPersonProjectile.cpp b/Templates/TP_FirstPerson/Source/TP_FirstPerson/TP_FirstPersonProjectile.cpp
--- a/Templates/TP_FirstPerson/Source/TP_FirstPerson/TP_FirstPersonProjectile.cpp
+++ b/Templates/TP_FirstPerson/Source/TP_FirstPerson/TP_FirstPersonProjectile.cpp
@@ -1,4 +1,4 @@
-// Copyright 1998-2018 Epic Games, Inc. All Rights Reserved.
+// Copyright 1998-2019 Epic Games, Inc. All Rights Reserved.
 
 #include "TP_FirstPersonProjectile.h"
 #include "GameFramework/ProjectileMovementComponent.h"
diff --git a/Templates/TP_FirstPerson/Source/TP_FirstPerson/TP_FirstPersonProjectile.h b/Templates/TP_FirstPerson/Source/TP_FirstPerson/TP_FirstPersonProjectile.h
--- a/Templates/TP_FirstPerson/Source/TP_FirstPerson/TP_FirstPersonProjectile.h
+++ b/Templates/TP_FirstPerson/Source/TP_FirstPerson/TP_FirstPersonProjectile.h
@@ -1,4 +1,4 @@
-// Copyright 1998-2018 Epic Games, Inc. All Rights Reserved.
+// Copyright 1998-2019 Epic Games, Inc. All Rights Reserved.
 
 #pragma once
 
diff --git a/Templates/TP_FirstPerson/Source/TP_FirstPersonEditor.Target.cs b/Templates/TP_FirstPerson/Source/TP_FirstPersonEditor.Target.cs
--- a/Templates/TP_FirstPerson/Source/TP_FirstPersonEditor.Target.cs
+++ b/Templates/TP_FirstPerson/Source/TP_FirstPersonEditor.Target.cs
@@ -1,4 +1,4 @@
-// Copyright 1998-2018 Epic Games, Inc. All Rights Reserved.
+// Copyright 1998-2019 Epic Games, Inc. All Rights Reserved.
 
 using UnrealBuildTool;
 using System.Collections.Generic;
diff --git a/Templates/TP_Flying/Source/TP_Flying.Target.cs b/Templates/TP_Flying/Source/TP_Flying.Target.cs
--- a/Templates/TP_Flying/Source/TP_Flying.Target.cs
+++ b/Templates/TP_Flying/Source/TP_Flying.Target.cs
@@ -1,4 +1,4 @@
-// Copyright 1998-2018 Epic Games, Inc. All Rights Reserved.
+// Copyright 1998-2019 Epic Games, Inc. All Rights Reserved.
 
 using UnrealBuildTool;
 using System.Collections.Generic;
diff --git a/Templates/TP_Flying/Source/TP_Flying/TP_Flying.Build.cs b/Templates/TP_Flying/Source/TP_Flying/TP_Flying.Build.cs
--- a/Templates/TP_Flying/Source/TP_Flying/TP_Flying.Build.cs
+++ b/Templates/TP_Flying/Source/TP_Flying/TP_Flying.Build.cs
@@ -1,4 +1,4 @@
-// Copyright 1998-2018 Epic Games, Inc. All Rights Reserved.
+// Copyright 1998-2019 Epic Games, Inc. All Rights Reserved.
 
 using UnrealBuildTool;
 
diff --git a/Templates/TP_Flying/Source/TP_Flying/TP_Flying.cpp b/Templates/TP_Flying/Source/TP_Flying/TP_Flying.cpp
--- a/Templates/TP_Flying/Source/TP_Flying/TP_Flying.cpp
+++ b/Templates/TP_Flying/Source/TP_Flying/TP_Flying.cpp
@@ -1,4 +1,4 @@
-// Copyright 1998-2018 Epic Games, Inc. All Rights Reserved.
+// Copyright 1998-2019 Epic Games, Inc. All Rights Reserved.
 
 #include "TP_Flying.h"
 #include "Modules/ModuleManager.h"
diff --git a/Templates/TP_Flying/Source/TP_Flying/TP_Flying.h b/Templates/TP_Flying/Source/TP_Flying/TP_Flying.h
--- a/Templates/TP_Flying/Source/TP_Flying/TP_Flying.h
+++ b/Templates/TP_Flying/Source/TP_Flying/TP_Flying.h
@@ -1,4 +1,4 @@
-// Copyright 1998-2018 Epic Games, Inc. All Rights Reserved.
+// Copyright 1998-2019 Epic Games, Inc. All Rights Reserved.
 
 #pragma once
 
diff --git a/Templates/TP_Flying/Source/TP_Flying/TP_FlyingGameMode.cpp b/Templates/TP_Flying/Source/TP_Flying/TP_FlyingGameMode.cpp
--- a/Templates/TP_Flying/Source/TP_Flying/TP_FlyingGameMode.cpp
+++ b/Templates/TP_Flying/Source/TP_Flying/TP_FlyingGameMode.cpp
@@ -1,4 +1,4 @@
-// Copyright 1998-2018 Epic Games, Inc. All Rights Reserved.
+// Copyright 1998-2019 Epic Games, Inc. All Rights Reserved.
 
 #include "TP_FlyingGameMode.h"
 #include "TP_FlyingPawn.h"
diff --git a/Templates/TP_Flying/Source/TP_Flying/TP_FlyingGameMode.h b/Templates/TP_Flying/Source/TP_Flying/TP_FlyingGameMode.h
--- a/Templates/TP_Flying/Source/TP_Flying/TP_FlyingGameMode.h
+++ b/Templates/TP_Flying/Source/TP_Flying/TP_FlyingGameMode.h
@@ -1,4 +1,4 @@
-// Copyright 1998-2018 Epic Games, Inc. All Rights Reserved.
+// Copyright 1998-2019 Epic Games, Inc. All Rights Reserved.
 
 #pragma once
 
diff --git a/Templates/TP_Flying/Source/TP_Flying/TP_FlyingPawn.cpp b/Templates/TP_Flying/Source/TP_Flying/TP_FlyingPawn.cpp
--- a/Templates/TP_Flying/Source/TP_Flying/TP_FlyingPawn.cpp
+++ b/Templates/TP_Flying/Source/TP_Flying/TP_FlyingPawn.cpp
@@ -1,4 +1,4 @@
-// Copyright 1998-2018 Epic Games, Inc. All Rights Reserved.
+// Copyright 1998-2019 Epic Games, Inc. All Rights Reserved.
 
 #include "TP_FlyingPawn.h"
 #include "UObject/ConstructorHelpers.h"
diff --git a/Templates/TP_Flying/Source/TP_Flying/TP_FlyingPawn.h b/Templates/TP_Flying/Source/TP_Flying/TP_FlyingPawn.h
--- a/Templates/TP_Flying/Source/TP_Flying/TP_FlyingPawn.h
+++ b/Templates/TP_Flying/Source/TP_Flying/TP_FlyingPawn.h
@@ -1,4 +1,4 @@
-// Copyright 1998-2018 Epic Games, Inc. All Rights Reserved.
+// Copyright 1998-2019 Epic Games, Inc. All Rights Reserved.
 #pragma once
 
 #include "CoreMinimal.h"
diff --git a/Templates/TP_Flying/Source/TP_FlyingEditor.Target.cs b/Templates/TP_Flying/Source/TP_FlyingEditor.Target.cs
--- a/Templates/TP_Flying/Source/TP_FlyingEditor.Target.cs
+++ b/Templates/TP_Flying/Source/TP_FlyingEditor.Target.cs
@@ -1,4 +1,4 @@
-// Copyright 1998-2018 Epic Games, Inc. All Rights Reserved.
+// Copyright 1998-2019 Epic Games, Inc. All Rights Reserved.
 
 using UnrealBuildTool;
 using System.Collections.Generic;
diff --git a/Templates/TP_Puzzle/Source/TP_Puzzle.Target.cs b/Templates/TP_Puzzle/Source/TP_Puzzle.Target.cs
--- a/Templates/TP_Puzzle/Source/TP_Puzzle.Target.cs
+++ b/Templates/TP_Puzzle/Source/TP_Puzzle.Target.cs
@@ -1,4 +1,4 @@
-// Copyright 1998-2018 Epic Games, Inc. All Rights Reserved.
+// Copyright 1998-2019 Epic Games, Inc. All Rights Reserved.
 
 using UnrealBuildTool;
 using System.Collections.Generic;
diff --git a/Templates/TP_Puzzle/Source/TP_Puzzle/TP_Puzzle.Build.cs b/Templates/TP_Puzzle/Source/TP_Puzzle/TP_Puzzle.Build.cs
--- a/Templates/TP_Puzzle/Source/TP_Puzzle/TP_Puzzle.Build.cs
+++ b/Templates/TP_Puzzle/Source/TP_Puzzle/TP_Puzzle.Build.cs
@@ -1,4 +1,4 @@
-// Copyright 1998-2018 Epic Games, Inc. All Rights Reserved.
+// Copyright 1998-2019 Epic Games, Inc. All Rights Reserved.
 
 using UnrealBuildTool;
 
diff --git a/Templates/TP_Puzzle/Source/TP_Puzzle/TP_Puzzle.cpp b/Templates/TP_Puzzle/Source/TP_Puzzle/TP_Puzzle.cpp
--- a/Templates/TP_Puzzle/Source/TP_Puzzle/TP_Puzzle.cpp
+++ b/Templates/TP_Puzzle/Source/TP_Puzzle/TP_Puzzle.cpp
@@ -1,4 +1,4 @@
-// Copyright 1998-2018 Epic Games, Inc. All Rights Reserved.
+// Copyright 1998-2019 Epic Games, Inc. All Rights Reserved.
 
 #include "TP_Puzzle.h"
 #include "Modules/ModuleManager.h"
diff --git a/Templates/TP_Puzzle/Source/TP_Puzzle/TP_Puzzle.h b/Templates/TP_Puzzle/Source/TP_Puzzle/TP_Puzzle.h
--- a/Templates/TP_Puzzle/Source/TP_Puzzle/TP_Puzzle.h
+++ b/Templates/TP_Puzzle/Source/TP_Puzzle/TP_Puzzle.h
@@ -1,4 +1,4 @@
-// Copyright 1998-2018 Epic Games, Inc. All Rights Reserved.
+// Copyright 1998-2019 Epic Games, Inc. All Rights Reserved.
 
 #pragma once
 
diff --git a/Templates/TP_Puzzle/Source/TP_Puzzle/TP_PuzzleBlock.cpp b/Templates/TP_Puzzle/Source/TP_Puzzle/TP_PuzzleBlock.cpp
--- a/Templates/TP_Puzzle/Source/TP_Puzzle/TP_PuzzleBlock.cpp
+++ b/Templates/TP_Puzzle/Source/TP_Puzzle/TP_PuzzleBlock.cpp
@@ -1,4 +1,4 @@
-// Copyright 1998-2018 Epic Games, Inc. All Rights Reserved.
+// Copyright 1998-2019 Epic Games, Inc. All Rights Reserved.
 
 #include "TP_PuzzleBlock.h"
 #include "TP_PuzzleBlockGrid.h"
diff --git a/Templates/TP_Puzzle/Source/TP_Puzzle/TP_PuzzleBlock.h b/Templates/TP_Puzzle/Source/TP_Puzzle/TP_PuzzleBlock.h
--- a/Templates/TP_Puzzle/Source/TP_Puzzle/TP_PuzzleBlock.h
+++ b/Templates/TP_Puzzle/Source/TP_Puzzle/TP_PuzzleBlock.h
@@ -1,4 +1,4 @@
-// Copyright 1998-2018 Epic Games, Inc. All Rights Reserved.
+// Copyright 1998-2019 Epic Games, Inc. All Rights Reserved.
 
 #pragma once
 
diff --git a/Templates/TP_Puzzle/Source/TP_Puzzle/TP_PuzzleBlockGrid.cpp b/Templates/TP_Puzzle/Source/TP_Puzzle/TP_PuzzleBlockGrid.cpp
--- a/Templates/TP_Puzzle/Source/TP_Puzzle/TP_PuzzleBlockGrid.cpp
+++ b/Templates/TP_Puzzle/Source/TP_Puzzle/TP_PuzzleBlockGrid.cpp
@@ -1,4 +1,4 @@
-// Copyright 1998-2018 Epic Games, Inc. All Rights Reserved.
+// Copyright 1998-2019 Epic Games, Inc. All Rights Reserved.
 
 #include "TP_PuzzleBlockGrid.h"
 #include "TP_PuzzleBlock.h"
diff --git a/Templates/TP_Puzzle/Source/TP_Puzzle/TP_PuzzleBlockGrid.h b/Templates/TP_Puzzle/Source/TP_Puzzle/TP_PuzzleBlockGrid.h
--- a/Templates/TP_Puzzle/Source/TP_Puzzle/TP_PuzzleBlockGrid.h
+++ b/Templates/TP_Puzzle/Source/TP_Puzzle/TP_PuzzleBlockGrid.h
@@ -1,4 +1,4 @@
-// Copyright 1998-2018 Epic Games, Inc. All Rights Reserved.
+// Copyright 1998-2019 Epic Games, Inc. All Rights Reserved.
 
 #pragma once
 
diff --git a/Templates/TP_Puzzle/Source/TP_Puzzle/TP_PuzzleGameMode.cpp b/Templates/TP_Puzzle/Source/TP_Puzzle/TP_PuzzleGameMode.cpp
--- a/Templates/TP_Puzzle/Source/TP_Puzzle/TP_PuzzleGameMode.cpp
+++ b/Templates/TP_Puzzle/Source/TP_Puzzle/TP_PuzzleGameMode.cpp
@@ -1,4 +1,4 @@
-// Copyright 1998-2018 Epic Games, Inc. All Rights Reserved.
+// Copyright 1998-2019 Epic Games, Inc. All Rights Reserved.
 
 #include "TP_PuzzleGameMode.h"
 #include "TP_PuzzlePlayerController.h"
diff --git a/Templates/TP_Puzzle/Source/TP_Puzzle/TP_PuzzleGameMode.h b/Templates/TP_Puzzle/Source/TP_Puzzle/TP_PuzzleGameMode.h
--- a/Templates/TP_Puzzle/Source/TP_Puzzle/TP_PuzzleGameMode.h
+++ b/Templates/TP_Puzzle/Source/TP_Puzzle/TP_PuzzleGameMode.h
@@ -1,4 +1,4 @@
-// Copyright 1998-2018 Epic Games, Inc. All Rights Reserved.
+// Copyright 1998-2019 Epic Games, Inc. All Rights Reserved.
 
 #pragma once
 
diff --git a/Templates/TP_Puzzle/Source/TP_Puzzle/TP_PuzzlePawn.cpp b/Templates/TP_Puzzle/Source/TP_Puzzle/TP_PuzzlePawn.cpp
--- a/Templates/TP_Puzzle/Source/TP_Puzzle/TP_PuzzlePawn.cpp
+++ b/Templates/TP_Puzzle/Source/TP_Puzzle/TP_PuzzlePawn.cpp
@@ -1,4 +1,4 @@
-// Copyright 1998-2018 Epic Games, Inc. All Rights Reserved.
+// Copyright 1998-2019 Epic Games, Inc. All Rights Reserved.
 
 #include "TP_PuzzlePawn.h"
 #include "TP_PuzzleBlock.h"
diff --git a/Templates/TP_Puzzle/Source/TP_Puzzle/TP_PuzzlePawn.h b/Templates/TP_Puzzle/Source/TP_Puzzle/TP_PuzzlePawn.h
--- a/Templates/TP_Puzzle/Source/TP_Puzzle/TP_PuzzlePawn.h
+++ b/Templates/TP_Puzzle/Source/TP_Puzzle/TP_PuzzlePawn.h
@@ -1,4 +1,4 @@
-// Copyright 1998-2018 Epic Games, Inc. All Rights Reserved.
+// Copyright 1998-2019 Epic Games, Inc. All Rights Reserved.
 
 #pragma once
 
diff --git a/Templates/TP_Puzzle/Source/TP_Puzzle/TP_PuzzlePlayerController.cpp b/Templates/TP_Puzzle/Source/TP_Puzzle/TP_PuzzlePlayerController.cpp
--- a/Templates/TP_Puzzle/Source/TP_Puzzle/TP_PuzzlePlayerController.cpp
+++ b/Templates/TP_Puzzle/Source/TP_Puzzle/TP_PuzzlePlayerController.cpp
@@ -1,4 +1,4 @@
-// Copyright 1998-2018 Epic Games, Inc. All Rights Reserved.
+// Copyright 1998-2019 Epic Games, Inc. All Rights Reserved.
 
 #include "TP_PuzzlePlayerController.h"
 
diff --git a/Templates/TP_Puzzle/Source/TP_Puzzle/TP_PuzzlePlayerController.h b/Templates/TP_Puzzle/Source/TP_Puzzle/TP_PuzzlePlayerController.h
--- a/Templates/TP_Puzzle/Source/TP_Puzzle/TP_PuzzlePlayerController.h
+++ b/Templates/TP_Puzzle/Source/TP_Puzzle/TP_PuzzlePlayerController.h
@@ -1,4 +1,4 @@
-// Copyright 1998-2018 Epic Games, Inc. All Rights Reserved.
+// Copyright 1998-2019 Epic Games, Inc. All Rights Reserved.
 
 #pragma once
 
diff --git a/Templates/TP_Puzzle/Source/TP_PuzzleEditor.Target.cs b/Templates/TP_Puzzle/Source/TP_PuzzleEditor.Target.cs
--- a/Templates/TP_Puzzle/Source/TP_PuzzleEditor.Target.cs
+++ b/Templates/TP_Puzzle/Source/TP_PuzzleEditor.Target.cs
@@ -1,4 +1,4 @@
-// Copyright 1998-2018 Epic Games, Inc. All Rights Reserved.
+// Copyright 1998-2019 Epic Games, Inc. All Rights Reserved.
 
 using UnrealBuildTool;
 using System.Collections.Generic;
diff --git a/Templates/TP_Rolling/Source/TP_Rolling.Target.cs b/Templates/TP_Rolling/Source/TP_Rolling.Target.cs
--- a/Templates/TP_Rolling/Source/TP_Rolling.Target.cs
+++ b/Templates/TP_Rolling/Source/TP_Rolling.Target.cs
@@ -1,4 +1,4 @@
-// Copyright 1998-2018 Epic Games, Inc. All Rights Reserved.
+// Copyright 1998-2019 Epic Games, Inc. All Rights Reserved.
 
 using UnrealBuildTool;
 using System.Collections.Generic;
diff --git a/Templates/TP_Rolling/Source/TP_Rolling/TP_Rolling.Build.cs b/Templates/TP_Rolling/Source/TP_Rolling/TP_Rolling.Build.cs
--- a/Templates/TP_Rolling/Source/TP_Rolling/TP_Rolling.Build.cs
+++ b/Templates/TP_Rolling/Source/TP_Rolling/TP_Rolling.Build.cs
@@ -1,4 +1,4 @@
-// Copyright 1998-2018 Epic Games, Inc. All Rights Reserved.
+// Copyright 1998-2019 Epic Games, Inc. All Rights Reserved.
 
 using UnrealBuildTool;
 
diff --git a/Templates/TP_Rolling/Source/TP_Rolling/TP_Rolling.cpp b/Templates/TP_Rolling/Source/TP_Rolling/TP_Rolling.cpp
--- a/Templates/TP_Rolling/Source/TP_Rolling/TP_Rolling.cpp
+++ b/Templates/TP_Rolling/Source/TP_Rolling/TP_Rolling.cpp
@@ -1,4 +1,4 @@
-// Copyright 1998-2018 Epic Games, Inc. All Rights Reserved.
+// Copyright 1998-2019 Epic Games, Inc. All Rights Reserved.
 
 #include "TP_Rolling.h"
 #include "Modules/ModuleManager.h"
diff --git a/Templates/TP_Rolling/Source/TP_Rolling/TP_Rolling.h b/Templates/TP_Rolling/Source/TP_Rolling/TP_Rolling.h
--- a/Templates/TP_Rolling/Source/TP_Rolling/TP_Rolling.h
+++ b/Templates/TP_Rolling/Source/TP_Rolling/TP_Rolling.h
@@ -1,4 +1,4 @@
-// Copyright 1998-2018 Epic Games, Inc. All Rights Reserved.
+// Copyright 1998-2019 Epic Games, Inc. All Rights Reserved.
 
 #pragma once
 
diff --git a/Templates/TP_Rolling/Source/TP_Rolling/TP_RollingBall.cpp b/Templates/TP_Rolling/Source/TP_Rolling/TP_RollingBall.cpp
--- a/Templates/TP_Rolling/Source/TP_Rolling/TP_RollingBall.cpp
+++ b/Templates/TP_Rolling/Source/TP_Rolling/TP_RollingBall.cpp
@@ -1,4 +1,4 @@
-// Copyright 1998-2018 Epic Games, Inc. All Rights Reserved.
+// Copyright 1998-2019 Epic Games, Inc. All Rights Reserved.
 
 #include "TP_RollingBall.h"
 #include "UObject/ConstructorHelpers.h"
diff --git a/Templates/TP_Rolling/Source/TP_Rolling/TP_RollingBall.h b/Templates/TP_Rolling/Source/TP_Rolling/TP_RollingBall.h
--- a/Templates/TP_Rolling/Source/TP_Rolling/TP_RollingBall.h
+++ b/Templates/TP_Rolling/Source/TP_Rolling/TP_RollingBall.h
@@ -1,4 +1,4 @@
-// Copyright 1998-2018 Epic Games, Inc. All Rights Reserved.
+// Copyright 1998-2019 Epic Games, Inc. All Rights Reserved.
 
 #pragma once
 
diff --git a/Templates/TP_Rolling/Source/TP_Rolling/TP_RollingGameMode.cpp b/Templates/TP_Rolling/Source/TP_Rolling/TP_RollingGameMode.cpp
--- a/Templates/TP_Rolling/Source/TP_Rolling/TP_RollingGameMode.cpp
+++ b/Templates/TP_Rolling/Source/TP_Rolling/TP_RollingGameMode.cpp
@@ -1,4 +1,4 @@
-// Copyright 1998-2018 Epic Games, Inc. All Rights Reserved.
+// Copyright 1998-2019 Epic Games, Inc. All Rights Reserved.
 
 #include "TP_RollingGameMode.h"
 #include "TP_RollingBall.h"
diff --git a/Templates/TP_Rolling/Source/TP_Rolling/TP_RollingGameMode.h b/Templates/TP_Rolling/Source/TP_Rolling/TP_RollingGameMode.h
--- a/Templates/TP_Rolling/Source/TP_Rolling/TP_RollingGameMode.h
+++ b/Templates/TP_Rolling/Source/TP_Rolling/TP_RollingGameMode.h
@@ -1,4 +1,4 @@
-// Copyright 1998-2018 Epic Games, Inc. All Rights Reserved.
+// Copyright 1998-2019 Epic Games, Inc. All Rights Reserved.
 
 #pragma once
 
diff --git a/Templates/TP_Rolling/Source/TP_RollingEditor.Target.cs b/Templates/TP_Rolling/Source/TP_RollingEditor.Target.cs
--- a/Templates/TP_Rolling/Source/TP_RollingEditor.Target.cs
+++ b/Templates/TP_Rolling/Source/TP_RollingEditor.Target.cs
@@ -1,4 +1,4 @@
-// Copyright 1998-2018 Epic Games, Inc. All Rights Reserved.
+// Copyright 1998-2019 Epic Games, Inc. All Rights Reserved.
 
 using UnrealBuildTool;
 using System.Collections.Generic;
diff --git a/Templates/TP_SideScroller/Source/TP_SideScroller.Target.cs b/Templates/TP_SideScroller/Source/TP_SideScroller.Target.cs
--- a/Templates/TP_SideScroller/Source/TP_SideScroller.Target.cs
+++ b/Templates/TP_SideScroller/Source/TP_SideScroller.Target.cs
@@ -1,4 +1,4 @@
-// Copyright 1998-2018 Epic Games, Inc. All Rights Reserved.
+// Copyright 1998-2019 Epic Games, Inc. All Rights Reserved.
 
 using UnrealBuildTool;
 using System.Collections.Generic;
diff --git a/Templates/TP_SideScroller/Source/TP_SideScroller/TP_SideScroller.Build.cs b/Templates/TP_SideScroller/Source/TP_SideScroller/TP_SideScroller.Build.cs
--- a/Templates/TP_SideScroller/Source/TP_SideScroller/TP_SideScroller.Build.cs
+++ b/Templates/TP_SideScroller/Source/TP_SideScroller/TP_SideScroller.Build.cs
@@ -1,4 +1,4 @@
-// Copyright 1998-2018 Epic Games, Inc. All Rights Reserved.
+// Copyright 1998-2019 Epic Games, Inc. All Rights Reserved.
 
 using UnrealBuildTool;
 
diff --git a/Templates/TP_SideScroller/Source/TP_SideScroller/TP_SideScroller.cpp b/Templates/TP_SideScroller/Source/TP_SideScroller/TP_SideScroller.cpp
--- a/Templates/TP_SideScroller/Source/TP_SideScroller/TP_SideScroller.cpp
+++ b/Templates/TP_SideScroller/Source/TP_SideScroller/TP_SideScroller.cpp
@@ -1,4 +1,4 @@
-// Copyright 1998-2018 Epic Games, Inc. All Rights Reserved.
+// Copyright 1998-2019 Epic Games, Inc. All Rights Reserved.
 
 #include "TP_SideScroller.h"
 #include "Modules/ModuleManager.h"
diff --git a/Templates/TP_SideScroller/Source/TP_SideScroller/TP_SideScroller.h b/Templates/TP_SideScroller/Source/TP_SideScroller/TP_SideScroller.h
--- a/Templates/TP_SideScroller/Source/TP_SideScroller/TP_SideScroller.h
+++ b/Templates/TP_SideScroller/Source/TP_SideScroller/TP_SideScroller.h
@@ -1,4 +1,4 @@
-// Copyright 1998-2018 Epic Games, Inc. All Rights Reserved.
+// Copyright 1998-2019 Epic Games, Inc. All Rights Reserved.
 
 #pragma once
 
diff --git a/Templates/TP_SideScroller/Source/TP_SideScroller/TP_SideScrollerCharacter.cpp b/Templates/TP_SideScroller/Source/TP_SideScroller/TP_SideScrollerCharacter.cpp
--- a/Templates/TP_SideScroller/Source/TP_SideScroller/TP_SideScrollerCharacter.cpp
+++ b/Templates/TP_SideScroller/Source/TP_SideScroller/TP_SideScrollerCharacter.cpp
@@ -1,4 +1,4 @@
-// Copyright 1998-2018 Epic Games, Inc. All Rights Reserved.
+// Copyright 1998-2019 Epic Games, Inc. All Rights Reserved.
 
 #include "TP_SideScrollerCharacter.h"
 #include "Camera/CameraComponent.h"
diff --git a/Templates/TP_SideScroller/Source/TP_SideScroller/TP_SideScrollerCharacter.h b/Templates/TP_SideScroller/Source/TP_SideScroller/TP_SideScrollerCharacter.h
--- a/Templates/TP_SideScroller/Source/TP_SideScroller/TP_SideScrollerCharacter.h
+++ b/Templates/TP_SideScroller/Source/TP_SideScroller/TP_SideScrollerCharacter.h
@@ -1,4 +1,4 @@
-// Copyright 1998-2018 Epic Games, Inc. All Rights Reserved.
+// Copyright 1998-2019 Epic Games, Inc. All Rights Reserved.
 
 #pragma once
 
diff --git a/Templates/TP_SideScroller/Source/TP_SideScroller/TP_SideScrollerGameMode.cpp b/Templates/TP_SideScroller/Source/TP_SideScroller/TP_SideScrollerGameMode.cpp
--- a/Templates/TP_SideScroller/Source/TP_SideScroller/TP_SideScrollerGameMode.cpp
+++ b/Templates/TP_SideScroller/Source/TP_SideScroller/TP_SideScrollerGameMode.cpp
@@ -1,4 +1,4 @@
-// Copyright 1998-2018 Epic Games, Inc. All Rights Reserved.
+// Copyright 1998-2019 Epic Games, Inc. All Rights Reserved.
 
 #include "TP_SideScrollerGameMode.h"
 #include "TP_SideScrollerCharacter.h"
diff --git a/Templates/TP_SideScroller/Source/TP_SideScroller/TP_SideScrollerGameMode.h b/Templates/TP_SideScroller/Source/TP_SideScroller/TP_SideScrollerGameMode.h
--- a/Templates/TP_SideScroller/Source/TP_SideScroller/TP_SideScrollerGameMode.h
+++ b/Templates/TP_SideScroller/Source/TP_SideScroller/TP_SideScrollerGameMode.h
@@ -1,4 +1,4 @@
-// Copyright 1998-2018 Epic Games, Inc. All Rights Reserved.
+// Copyright 1998-2019 Epic Games, Inc. All Rights Reserved.
 
 #pragma once
 
diff --git a/Templates/TP_SideScroller/Source/TP_SideScrollerEditor.Target.cs b/Templates/TP_SideScroller/Source/TP_SideScrollerEditor.Target.cs
--- a/Templates/TP_SideScroller/Source/TP_SideScrollerEditor.Target.cs
+++ b/Templates/TP_SideScroller/Source/TP_SideScrollerEditor.Target.cs
@@ -1,4 +1,4 @@
-// Copyright 1998-2018 Epic Games, Inc. All Rights Reserved.
+// Copyright 1998-2019 Epic Games, Inc. All Rights Reserved.
 
 using UnrealBuildTool;
 using System.Collections.Generic;
diff --git a/Templates/TP_ThirdPerson/Source/TP_ThirdPerson.Target.cs b/Templates/TP_ThirdPerson/Source/TP_ThirdPerson.Target.cs
--- a/Templates/TP_ThirdPerson/Source/TP_ThirdPerson.Target.cs
+++ b/Templates/TP_ThirdPerson/Source/TP_ThirdPerson.Target.cs
@@ -1,4 +1,4 @@
-// Copyright 1998-2018 Epic Games, Inc. All Rights Reserved.
+// Copyright 1998-2019 Epic Games, Inc. All Rights Reserved.
 
 using UnrealBuildTool;
 using System.Collections.Generic;
diff --git a/Templates/TP_ThirdPerson/Source/TP_ThirdPerson/TP_ThirdPerson.Build.cs b/Templates/TP_ThirdPerson/Source/TP_ThirdPerson/TP_ThirdPerson.Build.cs
--- a/Templates/TP_ThirdPerson/Source/TP_ThirdPerson/TP_ThirdPerson.Build.cs
+++ b/Templates/TP_ThirdPerson/Source/TP_ThirdPerson/TP_ThirdPerson.Build.cs
@@ -1,4 +1,4 @@
-// Copyright 1998-2018 Epic Games, Inc. All Rights Reserved.
+// Copyright 1998-2019 Epic Games, Inc. All Rights Reserved.
 
 using UnrealBuildTool;
 
diff --git a/Templates/TP_ThirdPerson/Source/TP_ThirdPerson/TP_ThirdPerson.cpp b/Templates/TP_ThirdPerson/Source/TP_ThirdPerson/TP_ThirdPerson.cpp
--- a/Templates/TP_ThirdPerson/Source/TP_ThirdPerson/TP_ThirdPerson.cpp
+++ b/Templates/TP_ThirdPerson/Source/TP_ThirdPerson/TP_ThirdPerson.cpp
@@ -1,4 +1,4 @@
-// Copyright 1998-2018 Epic Games, Inc. All Rights Reserved.
+// Copyright 1998-2019 Epic Games, Inc. All Rights Reserved.
 
 #include "TP_ThirdPerson.h"
 #include "Modules/ModuleManager.h"
diff --git a/Templates/TP_ThirdPerson/Source/TP_ThirdPerson/TP_ThirdPerson.h b/Templates/TP_ThirdPerson/Source/TP_ThirdPerson/TP_ThirdPerson.h
--- a/Templates/TP_ThirdPerson/Source/TP_ThirdPerson/TP_ThirdPerson.h
+++ b/Templates/TP_ThirdPerson/Source/TP_ThirdPerson/TP_ThirdPerson.h
@@ -1,4 +1,4 @@
-// Copyright 1998-2018 Epic Games, Inc. All Rights Reserved.
+// Copyright 1998-2019 Epic Games, Inc. All Rights Reserved.
 
 #pragma once
 
diff --git a/Templates/TP_ThirdPerson/Source/TP_ThirdPerson/TP_ThirdPersonCharacter.cpp b/Templates/TP_ThirdPerson/Source/TP_ThirdPerson/TP_ThirdPersonCharacter.cpp
--- a/Templates/TP_ThirdPerson/Source/TP_ThirdPerson/TP_ThirdPersonCharacter.cpp
+++ b/Templates/TP_ThirdPerson/Source/TP_ThirdPerson/TP_ThirdPersonCharacter.cpp
@@ -1,4 +1,4 @@
-// Copyright 1998-2018 Epic Games, Inc. All Rights Reserved.
+// Copyright 1998-2019 Epic Games, Inc. All Rights Reserved.
 
 #include "TP_ThirdPersonCharacter.h"
 #include "HeadMountedDisplayFunctionLibrary.h"
diff --git a/Templates/TP_ThirdPerson/Source/TP_ThirdPerson/TP_ThirdPersonCharacter.h b/Templates/TP_ThirdPerson/Source/TP_ThirdPerson/TP_ThirdPersonCharacter.h
--- a/Templates/TP_ThirdPerson/Source/TP_ThirdPerson/TP_ThirdPersonCharacter.h
+++ b/Templates/TP_ThirdPerson/Source/TP_ThirdPerson/TP_ThirdPersonCharacter.h
@@ -1,4 +1,4 @@
-// Copyright 1998-2018 Epic Games, Inc. All Rights Reserved.
+// Copyright 1998-2019 Epic Games, Inc. All Rights Reserved.
 
 #pragma once
 
diff --git a/Templates/TP_ThirdPerson/Source/TP_ThirdPerson/TP_ThirdPersonGameMode.cpp b/Templates/TP_ThirdPerson/Source/TP_ThirdPerson/TP_ThirdPersonGameMode.cpp
--- a/Templates/TP_ThirdPerson/Source/TP_ThirdPerson/TP_ThirdPersonGameMode.cpp
+++ b/Templates/TP_ThirdPerson/Source/TP_ThirdPerson/TP_ThirdPersonGameMode.cpp
@@ -1,4 +1,4 @@
-// Copyright 1998-2018 Epic Games, Inc. All Rights Reserved.
+// Copyright 1998-2019 Epic Games, Inc. All Rights Reserved.
 
 #include "TP_ThirdPersonGameMode.h"
 #include "TP_ThirdPersonCharacter.h"
diff --git a/Templates/TP_ThirdPerson/Source/TP_ThirdPerson/TP_ThirdPersonGameMode.h b/Templates/TP_ThirdPerson/Source/TP_ThirdPerson/TP_ThirdPersonGameMode.h
--- a/Templates/TP_ThirdPerson/Source/TP_ThirdPerson/TP_ThirdPersonGameMode.h
+++ b/Templates/TP_ThirdPerson/Source/TP_ThirdPerson/TP_ThirdPersonGameMode.h
@@ -1,4 +1,4 @@
-// Copyright 1998-2018 Epic Games, Inc. All Rights Reserved.
+// Copyright 1998-2019 Epic Games, Inc. All Rights Reserved.
 
 #pragma once
 
diff --git a/Templates/TP_ThirdPerson/Source/TP_ThirdPersonEditor.Target.cs b/Templates/TP_ThirdPerson/Source/TP_ThirdPersonEditor.Target.cs
--- a/Templates/TP_ThirdPerson/Source/TP_ThirdPersonEditor.Target.cs
+++ b/Templates/TP_ThirdPerson/Source/TP_ThirdPersonEditor.Target.cs
@@ -1,4 +1,4 @@
-// Copyright 1998-2018 Epic Games, Inc. All Rights Reserved.
+// Copyright 1998-2019 Epic Games, Inc. All Rights Reserved.
 
 using UnrealBuildTool;
 using System.Collections.Generic;
diff --git a/Templates/TP_TopDown/Source/TP_TopDown.Target.cs b/Templates/TP_TopDown/Source/TP_TopDown.Target.cs
--- a/Templates/TP_TopDown/Source/TP_TopDown.Target.cs
+++ b/Templates/TP_TopDown/Source/TP_TopDown.Target.cs
@@ -1,4 +1,4 @@
-// Copyright 1998-2018 Epic Games, Inc. All Rights Reserved.
+// Copyright 1998-2019 Epic Games, Inc. All Rights Reserved.
 
 using UnrealBuildTool;
 using System.Collections.Generic;
diff --git a/Templates/TP_TopDown/Source/TP_TopDown/TP_TopDown.Build.cs b/Templates/TP_TopDown/Source/TP_TopDown/TP_TopDown.Build.cs
--- a/Templates/TP_TopDown/Source/TP_TopDown/TP_TopDown.Build.cs
+++ b/Templates/TP_TopDown/Source/TP_TopDown/TP_TopDown.Build.cs
@@ -1,4 +1,4 @@
-// Copyright 1998-2018 Epic Games, Inc. All Rights Reserved.
+// Copyright 1998-2019 Epic Games, Inc. All Rights Reserved.
 
 using UnrealBuildTool;
 
diff --git a/Templates/TP_TopDown/Source/TP_TopDown/TP_TopDown.cpp b/Templates/TP_TopDown/Source/TP_TopDown/TP_TopDown.cpp
--- a/Templates/TP_TopDown/Source/TP_TopDown/TP_TopDown.cpp
+++ b/Templates/TP_TopDown/Source/TP_TopDown/TP_TopDown.cpp
@@ -1,4 +1,4 @@
-// Copyright 1998-2018 Epic Games, Inc. All Rights Reserved.
+// Copyright 1998-2019 Epic Games, Inc. All Rights Reserved.
 
 #include "TP_TopDown.h"
 #include "Modules/ModuleManager.h"
diff --git a/Templates/TP_TopDown/Source/TP_TopDown/TP_TopDown.h b/Templates/TP_TopDown/Source/TP_TopDown/TP_TopDown.h
--- a/Templates/TP_TopDown/Source/TP_TopDown/TP_TopDown.h
+++ b/Templates/TP_TopDown/Source/TP_TopDown/TP_TopDown.h
@@ -1,4 +1,4 @@
-// Copyright 1998-2018 Epic Games, Inc. All Rights Reserved.
+// Copyright 1998-2019 Epic Games, Inc. All Rights Reserved.
 
 #pragma once
 
diff --git a/Templates/TP_TopDown/Source/TP_TopDown/TP_TopDownCharacter.cpp b/Templates/TP_TopDown/Source/TP_TopDown/TP_TopDownCharacter.cpp
--- a/Templates/TP_TopDown/Source/TP_TopDown/TP_TopDownCharacter.cpp
+++ b/Templates/TP_TopDown/Source/TP_TopDown/TP_TopDownCharacter.cpp
@@ -1,4 +1,4 @@
-// Copyright 1998-2018 Epic Games, Inc. All Rights Reserved.
+// Copyright 1998-2019 Epic Games, Inc. All Rights Reserved.
 
 #include "TP_TopDownCharacter.h"
 #include "UObject/ConstructorHelpers.h"
diff --git a/Templates/TP_TopDown/Source/TP_TopDown/TP_TopDownCharacter.h b/Templates/TP_TopDown/Source/TP_TopDown/TP_TopDownCharacter.h
--- a/Templates/TP_TopDown/Source/TP_TopDown/TP_TopDownCharacter.h
+++ b/Templates/TP_TopDown/Source/TP_TopDown/TP_TopDownCharacter.h
@@ -1,4 +1,4 @@
-// Copyright 1998-2018 Epic Games, Inc. All Rights Reserved.
+// Copyright 1998-2019 Epic Games, Inc. All Rights Reserved.
 
 #pragma once
 
diff --git a/Templates/TP_TopDown/Source/TP_TopDown/TP_TopDownGameMode.cpp b/Templates/TP_TopDown/Source/TP_TopDown/TP_TopDownGameMode.cpp
--- a/Templates/TP_TopDown/Source/TP_TopDown/TP_TopDownGameMode.cpp
+++ b/Templates/TP_TopDown/Source/TP_TopDown/TP_TopDownGameMode.cpp
@@ -1,4 +1,4 @@
-// Copyright 1998-2018 Epic Games, Inc. All Rights Reserved.
+// Copyright 1998-2019 Epic Games, Inc. All Rights Reserved.
 
 #include "TP_TopDownGameMode.h"
 #include "TP_TopDownPlayerController.h"
diff --git a/Templates/TP_TopDown/Source/TP_TopDown/TP_TopDownGameMode.h b/Templates/TP_TopDown/Source/TP_TopDown/TP_TopDownGameMode.h
--- a/Templates/TP_TopDown/Source/TP_TopDown/TP_TopDownGameMode.h
+++ b/Templates/TP_TopDown/Source/TP_TopDown/TP_TopDownGameMode.h
@@ -1,4 +1,4 @@
-// Copyright 1998-2018 Epic Games, Inc. All Rights Reserved.
+// Copyright 1998-2019 Epic Games, Inc. All Rights Reserved.
 
 #pragma once
 
diff --git a/Templates/TP_TopDown/Source/TP_TopDown/TP_TopDownPlayerController.cpp b/Templates/TP_TopDown/Source/TP_TopDown/TP_TopDownPlayerController.cpp
--- a/Templates/TP_TopDown/Source/TP_TopDown/TP_TopDownPlayerController.cpp
+++ b/Templates/TP_TopDown/Source/TP_TopDown/TP_TopDownPlayerController.cpp
@@ -1,4 +1,4 @@
-// Copyright 1998-2018 Epic Games, Inc. All Rights Reserved.
+// Copyright 1998-2019 Epic Games, Inc. All Rights Reserved.
 
 #include "TP_TopDownPlayerController.h"
 #include "Blueprint/AIBlueprintHelperLibrary.h"
diff --git a/Templates/TP_TopDown/Source/TP_TopDown/TP_TopDownPlayerController.h b/Templates/TP_TopDown/Source/TP_TopDown/TP_TopDownPlayerController.h
--- a/Templates/TP_TopDown/Source/TP_TopDown/TP_TopDownPlayerController.h
+++ b/Templates/TP_TopDown/Source/TP_TopDown/TP_TopDownPlayerController.h
@@ -1,4 +1,4 @@
-// Copyright 1998-2018 Epic Games, Inc. All Rights Reserved.
+// Copyright 1998-2019 Epic Games, Inc. All Rights Reserved.
 
 #pragma once
 
diff --git a/Templates/TP_TopDown/Source/TP_TopDownEditor.Target.cs b/Templates/TP_TopDown/Source/TP_TopDownEditor.Target.cs
--- a/Templates/TP_TopDown/Source/TP_TopDownEditor.Target.cs
+++ b/Templates/TP_TopDown/Source/TP_TopDownEditor.Target.cs
@@ -1,4 +1,4 @@
-// Copyright 1998-2018 Epic Games, Inc. All Rights Reserved.
+// Copyright 1998-2019 Epic Games, Inc. All Rights Reserved.
 
 using UnrealBuildTool;
 using System.Collections.Generic;
diff --git a/Templates/TP_TwinStick/Source/TP_TwinStick.Target.cs b/Templates/TP_TwinStick/Source/TP_TwinStick.Target.cs
--- a/Templates/TP_TwinStick/Source/TP_TwinStick.Target.cs
+++ b/Templates/TP_TwinStick/Source/TP_TwinStick.Target.cs
@@ -1,4 +1,4 @@
-// Copyright 1998-2018 Epic Games, Inc. All Rights Reserved.
+// Copyright 1998-2019 Epic Games, Inc. All Rights Reserved.
 
 using UnrealBuildTool;
 using System.Collections.Generic;
diff --git a/Templates/TP_TwinStick/Source/TP_TwinStick/TP_TwinStick.Build.cs b/Templates/TP_TwinStick/Source/TP_TwinStick/TP_TwinStick.Build.cs
--- a/Templates/TP_TwinStick/Source/TP_TwinStick/TP_TwinStick.Build.cs
+++ b/Templates/TP_TwinStick/Source/TP_TwinStick/TP_TwinStick.Build.cs
@@ -1,4 +1,4 @@
-// Copyright 1998-2018 Epic Games, Inc. All Rights Reserved.
+// Copyright 1998-2019 Epic Games, Inc. All Rights Reserved.
 
 using UnrealBuildTool;
 
diff --git a/Templates/TP_TwinStick/Source/TP_TwinStick/TP_TwinStick.cpp b/Templates/TP_TwinStick/Source/TP_TwinStick/TP_TwinStick.cpp
--- a/Templates/TP_TwinStick/Source/TP_TwinStick/TP_TwinStick.cpp
+++ b/Templates/TP_TwinStick/Source/TP_TwinStick/TP_TwinStick.cpp
@@ -1,4 +1,4 @@
-// Copyright 1998-2018 Epic Games, Inc. All Rights Reserved.
+// Copyright 1998-2019 Epic Games, Inc. All Rights Reserved.
 
 #include "TP_TwinStick.h"
 #include "Modules/ModuleManager.h"
diff --git a/Templates/TP_TwinStick/Source/TP_TwinStick/TP_TwinStick.h b/Templates/TP_TwinStick/Source/TP_TwinStick/TP_TwinStick.h
--- a/Templates/TP_TwinStick/Source/TP_TwinStick/TP_TwinStick.h
+++ b/Templates/TP_TwinStick/Source/TP_TwinStick/TP_TwinStick.h
@@ -1,4 +1,4 @@
-// Copyright 1998-2018 Epic Games, Inc. All Rights Reserved.
+// Copyright 1998-2019 Epic Games, Inc. All Rights Reserved.
 
 #pragma once
 
diff --git a/Templates/TP_TwinStick/Source/TP_TwinStick/TP_TwinStickGameMode.cpp b/Templates/TP_TwinStick/Source/TP_TwinStick/TP_TwinStickGameMode.cpp
--- a/Templates/TP_TwinStick/Source/TP_TwinStick/TP_TwinStickGameMode.cpp
+++ b/Templates/TP_TwinStick/Source/TP_TwinStick/TP_TwinStickGameMode.cpp
@@ -1,4 +1,4 @@
-// Copyright 1998-2018 Epic Games, Inc. All Rights Reserved.
+// Copyright 1998-2019 Epic Games, Inc. All Rights Reserved.
 
 #include "TP_TwinStickGameMode.h"
 #include "TP_TwinStickPawn.h"
diff --git a/Templates/TP_TwinStick/Source/TP_TwinStick/TP_TwinStickGameMode.h b/Templates/TP_TwinStick/Source/TP_TwinStick/TP_TwinStickGameMode.h
--- a/Templates/TP_TwinStick/Source/TP_TwinStick/TP_TwinStickGameMode.h
+++ b/Templates/TP_TwinStick/Source/TP_TwinStick/TP_TwinStickGameMode.h
@@ -1,4 +1,4 @@
-// Copyright 1998-2018 Epic Games, Inc. All Rights Reserved.
+// Copyright 1998-2019 Epic Games, Inc. All Rights Reserved.
 
 #pragma once
 
diff --git a/Templates/TP_TwinStick/Source/TP_TwinStick/TP_TwinStickPawn.cpp b/Templates/TP_TwinStick/Source/TP_TwinStick/TP_TwinStickPawn.cpp
--- a/Templates/TP_TwinStick/Source/TP_TwinStick/TP_TwinStickPawn.cpp
+++ b/Templates/TP_TwinStick/Source/TP_TwinStick/TP_TwinStickPawn.cpp
@@ -1,4 +1,4 @@
-// Copyright 1998-2018 Epic Games, Inc. All Rights Reserved.
+// Copyright 1998-2019 Epic Games, Inc. All Rights Reserved.
 
 #include "TP_TwinStickPawn.h"
 #include "TP_TwinStickProjectile.h"
diff --git a/Templates/TP_TwinStick/Source/TP_TwinStick/TP_TwinStickPawn.h b/Templates/TP_TwinStick/Source/TP_TwinStick/TP_TwinStickPawn.h
--- a/Templates/TP_TwinStick/Source/TP_TwinStick/TP_TwinStickPawn.h
+++ b/Templates/TP_TwinStick/Source/TP_TwinStick/TP_TwinStickPawn.h
@@ -1,4 +1,4 @@
-// Copyright 1998-2018 Epic Games, Inc. All Rights Reserved.
+// Copyright 1998-2019 Epic Games, Inc. All Rights Reserved.
 
 #pragma once
 
diff --git a/Templates/TP_TwinStick/Source/TP_TwinStick/TP_TwinStickProjectile.cpp b/Templates/TP_TwinStick/Source/TP_TwinStick/TP_TwinStickProjectile.cpp
--- a/Templates/TP_TwinStick/Source/TP_TwinStick/TP_TwinStickProjectile.cpp
+++ b/Templates/TP_TwinStick/Source/TP_TwinStick/TP_TwinStickProjectile.cpp
@@ -1,4 +1,4 @@
-// Copyright 1998-2018 Epic Games, Inc. All Rights Reserve
+// Copyright 1998-2019 Epic Games, Inc. All Rights Reserve
 
 #include "TP_TwinStickProjectile.h"
 #include "GameFramework/ProjectileMovementComponent.h"
diff --git a/Templates/TP_TwinStick/Source/TP_TwinStick/TP_TwinStickProjectile.h b/Templates/TP_TwinStick/Source/TP_TwinStick/TP_TwinStickProjectile.h
--- a/Templates/TP_TwinStick/Source/TP_TwinStick/TP_TwinStickProjectile.h
+++ b/Templates/TP_TwinStick/Source/TP_TwinStick/TP_TwinStickProjectile.h
@@ -1,4 +1,4 @@
-// Copyright 1998-2018 Epic Games, Inc. All Rights Reserved.
+// Copyright 1998-2019 Epic Games, Inc. All Rights Reserved.
 
 #pragma once
 
diff --git a/Templates/TP_TwinStick/Source/TP_TwinStickEditor.Target.cs b/Templates/TP_TwinStick/Source/TP_TwinStickEditor.Target.cs
--- a/Templates/TP_TwinStick/Source/TP_TwinStickEditor.Target.cs
+++ b/Templates/TP_TwinStick/Source/TP_TwinStickEditor.Target.cs
@@ -1,4 +1,4 @@
-// Copyright 1998-2018 Epic Games, Inc. All Rights Reserved.
+// Copyright 1998-2019 Epic Games, Inc. All Rights Reserved.
 
 using UnrealBuildTool;
 using System.Collections.Generic;
diff --git a/Templates/TP_Vehicle/Source/TP_Vehicle.Target.cs b/Templates/TP_Vehicle/Source/TP_Vehicle.Target.cs
--- a/Templates/TP_Vehicle/Source/TP_Vehicle.Target.cs
+++ b/Templates/TP_Vehicle/Source/TP_Vehicle.Target.cs
@@ -1,4 +1,4 @@
-// Copyright 1998-2018 Epic Games, Inc. All Rights Reserved.
+// Copyright 1998-2019 Epic Games, Inc. All Rights Reserved.
 
 using UnrealBuildTool;
 using System.Collections.Generic;
diff --git a/Templates/TP_Vehicle/Source/TP_Vehicle/TP_Vehicle.Build.cs b/Templates/TP_Vehicle/Source/TP_Vehicle/TP_Vehicle.Build.cs
--- a/Templates/TP_Vehicle/Source/TP_Vehicle/TP_Vehicle.Build.cs
+++ b/Templates/TP_Vehicle/Source/TP_Vehicle/TP_Vehicle.Build.cs
@@ -1,4 +1,4 @@
-// Copyright 1998-2018 Epic Games, Inc. All Rights Reserved.
+// Copyright 1998-2019 Epic Games, Inc. All Rights Reserved.
 
 using UnrealBuildTool;
 
diff --git a/Templates/TP_Vehicle/Source/TP_Vehicle/TP_Vehicle.cpp b/Templates/TP_Vehicle/Source/TP_Vehicle/TP_Vehicle.cpp
--- a/Templates/TP_Vehicle/Source/TP_Vehicle/TP_Vehicle.cpp
+++ b/Templates/TP_Vehicle/Source/TP_Vehicle/TP_Vehicle.cpp
@@ -1,4 +1,4 @@
-// Copyright 1998-2018 Epic Games, Inc. All Rights Reserved.
+// Copyright 1998-2019 Epic Games, Inc. All Rights Reserved.
 
 #include "TP_Vehicle.h"
 #include "Modules/ModuleManager.h"
diff --git a/Templates/TP_Vehicle/Source/TP_Vehicle/TP_Vehicle.h b/Templates/TP_Vehicle/Source/TP_Vehicle/TP_Vehicle.h
--- a/Templates/TP_Vehicle/Source/TP_Vehicle/TP_Vehicle.h
+++ b/Templates/TP_Vehicle/Source/TP_Vehicle/TP_Vehicle.h
@@ -1,4 +1,4 @@
-// Copyright 1998-2018 Epic Games, Inc. All Rights Reserved.
+// Copyright 1998-2019 Epic Games, Inc. All Rights Reserved.
 
 #pragma once
 
diff --git a/Templates/TP_Vehicle/Source/TP_Vehicle/TP_VehicleGameMode.cpp b/Templates/TP_Vehicle/Source/TP_Vehicle/TP_VehicleGameMode.cpp
--- a/Templates/TP_Vehicle/Source/TP_Vehicle/TP_VehicleGameMode.cpp
+++ b/Templates/TP_Vehicle/Source/TP_Vehicle/TP_VehicleGameMode.cpp
@@ -1,4 +1,4 @@
-// Copyright 1998-2018 Epic Games, Inc. All Rights Reserved.
+// Copyright 1998-2019 Epic Games, Inc. All Rights Reserved.
 
 #include "TP_VehicleGameMode.h"
 #include "TP_VehiclePawn.h"
diff --git a/Templates/TP_Vehicle/Source/TP_Vehicle/TP_VehicleGameMode.h b/Templates/TP_Vehicle/Source/TP_Vehicle/TP_VehicleGameMode.h
--- a/Templates/TP_Vehicle/Source/TP_Vehicle/TP_VehicleGameMode.h
+++ b/Templates/TP_Vehicle/Source/TP_Vehicle/TP_VehicleGameMode.h
@@ -1,4 +1,4 @@
-// Copyright 1998-2018 Epic Games, Inc. All Rights Reserved.
+// Copyright 1998-2019 Epic Games, Inc. All Rights Reserved.
 #pragma once
 #include "GameFramework/GameModeBase.h"
 #include "TP_VehicleGameMode.generated.h"
diff --git a/Templates/TP_Vehicle/Source/TP_Vehicle/TP_VehicleHud.cpp b/Templates/TP_Vehicle/Source/TP_Vehicle/TP_VehicleHud.cpp
--- a/Templates/TP_Vehicle/Source/TP_Vehicle/TP_VehicleHud.cpp
+++ b/Templates/TP_Vehicle/Source/TP_Vehicle/TP_VehicleHud.cpp
@@ -1,4 +1,4 @@
-// Copyright 1998-2018 Epic Games, Inc. All Rights Reserved.
+// Copyright 1998-2019 Epic Games, Inc. All Rights Reserved.
 
 #include "TP_VehicleHud.h"
 #include "TP_VehiclePawn.h"
diff --git a/Templates/TP_Vehicle/Source/TP_Vehicle/TP_VehicleHud.h b/Templates/TP_Vehicle/Source/TP_Vehicle/TP_VehicleHud.h
--- a/Templates/TP_Vehicle/Source/TP_Vehicle/TP_VehicleHud.h
+++ b/Templates/TP_Vehicle/Source/TP_Vehicle/TP_VehicleHud.h
@@ -1,4 +1,4 @@
-// Copyright 1998-2018 Epic Games, Inc. All Rights Reserved.
+// Copyright 1998-2019 Epic Games, Inc. All Rights Reserved.
 #pragma once
 #include "GameFramework/HUD.h"
 #include "TP_VehicleHud.generated.h"
diff --git a/Templates/TP_Vehicle/Source/TP_Vehicle/TP_VehiclePawn.cpp b/Templates/TP_Vehicle/Source/TP_Vehicle/TP_VehiclePawn.cpp
--- a/Templates/TP_Vehicle/Source/TP_Vehicle/TP_VehiclePawn.cpp
+++ b/Templates/TP_Vehicle/Source/TP_Vehicle/TP_VehiclePawn.cpp
@@ -1,4 +1,4 @@
-// Copyright 1998-2018 Epic Games, Inc. All Rights Reserved.
+// Copyright 1998-2019 Epic Games, Inc. All Rights Reserved.
 
 #include "TP_VehiclePawn.h"
 #include "TP_VehicleWheelFront.h"
diff --git a/Templates/TP_Vehicle/Source/TP_Vehicle/TP_VehiclePawn.h b/Templates/TP_Vehicle/Source/TP_Vehicle/TP_VehiclePawn.h
--- a/Templates/TP_Vehicle/Source/TP_Vehicle/TP_VehiclePawn.h
+++ b/Templates/TP_Vehicle/Source/TP_Vehicle/TP_VehiclePawn.h
@@ -1,4 +1,4 @@
-// Copyright 1998-2018 Epic Games, Inc. All Rights Reserved.
+// Copyright 1998-2019 Epic Games, Inc. All Rights Reserved.
 
 #pragma once
 
diff --git a/Templates/TP_Vehicle/Source/TP_Vehicle/TP_VehicleWheelFront.cpp b/Templates/TP_Vehicle/Source/TP_Vehicle/TP_VehicleWheelFront.cpp
--- a/Templates/TP_Vehicle/Source/TP_Vehicle/TP_VehicleWheelFront.cpp
+++ b/Templates/TP_Vehicle/Source/TP_Vehicle/TP_VehicleWheelFront.cpp
@@ -1,4 +1,4 @@
-// Copyright 1998-2018 Epic Games, Inc. All Rights Reserved.
+// Copyright 1998-2019 Epic Games, Inc. All Rights Reserved.
 
 #include "TP_VehicleWheelFront.h"
 
diff --git a/Templates/TP_Vehicle/Source/TP_Vehicle/TP_VehicleWheelFront.h b/Templates/TP_Vehicle/Source/TP_Vehicle/TP_VehicleWheelFront.h
--- a/Templates/TP_Vehicle/Source/TP_Vehicle/TP_VehicleWheelFront.h
+++ b/Templates/TP_Vehicle/Source/TP_Vehicle/TP_VehicleWheelFront.h
@@ -1,4 +1,4 @@
-// Copyright 1998-2018 Epic Games, Inc. All Rights Reserved.
+// Copyright 1998-2019 Epic Games, Inc. All Rights Reserved.
 
 #pragma once
 
diff --git a/Templates/TP_Vehicle/Source/TP_Vehicle/TP_VehicleWheelRear.cpp b/Templates/TP_Vehicle/Source/TP_Vehicle/TP_VehicleWheelRear.cpp
--- a/Templates/TP_Vehicle/Source/TP_Vehicle/TP_VehicleWheelRear.cpp
+++ b/Templates/TP_Vehicle/Source/TP_Vehicle/TP_VehicleWheelRear.cpp
@@ -1,4 +1,4 @@
-// Copyright 1998-2018 Epic Games, Inc. All Rights Reserved.
+// Copyright 1998-2019 Epic Games, Inc. All Rights Reserved.
 
 #include "TP_VehicleWheelRear.h"
 
diff --git a/Templates/TP_Vehicle/Source/TP_Vehicle/TP_VehicleWheelRear.h b/Templates/TP_Vehicle/Source/TP_Vehicle/TP_VehicleWheelRear.h
--- a/Templates/TP_Vehicle/Source/TP_Vehicle/TP_VehicleWheelRear.h
+++ b/Templates/TP_Vehicle/Source/TP_Vehicle/TP_VehicleWheelRear.h
@@ -1,4 +1,4 @@
-// Copyright 1998-2018 Epic Games, Inc. All Rights Reserved.
+// Copyright 1998-2019 Epic Games, Inc. All Rights Reserved.
 
 #pragma once
 
diff --git a/Templates/TP_Vehicle/Source/TP_VehicleEditor.Target.cs b/Templates/TP_Vehicle/Source/TP_VehicleEditor.Target.cs
--- a/Templates/TP_Vehicle/Source/TP_VehicleEditor.Target.cs
+++ b/Templates/TP_Vehicle/Source/TP_VehicleEditor.Target.cs
@@ -1,4 +1,4 @@
-// Copyright 1998-2018 Epic Games, Inc. All Rights Reserved.
+// Copyright 1998-2019 Epic Games, Inc. All Rights Reserved.
 
 using UnrealBuildTool;
 using System.Collections.Generic;
diff --git a/Templates/TP_VehicleAdv/Source/TP_VehicleAdv.Target.cs b/Templates/TP_VehicleAdv/Source/TP_VehicleAdv.Target.cs
--- a/Templates/TP_VehicleAdv/Source/TP_VehicleAdv.Target.cs
+++ b/Templates/TP_VehicleAdv/Source/TP_VehicleAdv.Target.cs
@@ -1,4 +1,4 @@
-// Copyright 1998-2018 Epic Games, Inc. All Rights Reserved.
+// Copyright 1998-2019 Epic Games, Inc. All Rights Reserved.
 
 using UnrealBuildTool;
 using System.Collections.Generic;
diff --git a/Templates/TP_VehicleAdv/Source/TP_VehicleAdv/TP_VehicleAdv.Build.cs b/Templates/TP_VehicleAdv/Source/TP_VehicleAdv/TP_VehicleAdv.Build.cs
--- a/Templates/TP_VehicleAdv/Source/TP_VehicleAdv/TP_VehicleAdv.Build.cs
+++ b/Templates/TP_VehicleAdv/Source/TP_VehicleAdv/TP_VehicleAdv.Build.cs
@@ -1,4 +1,4 @@
-// Copyright 1998-2018 Epic Games, Inc. All Rights Reserved.
+// Copyright 1998-2019 Epic Games, Inc. All Rights Reserved.
 
 using UnrealBuildTool;
 
diff --git a/Templates/TP_VehicleAdv/Source/TP_VehicleAdv/TP_VehicleAdv.cpp b/Templates/TP_VehicleAdv/Source/TP_VehicleAdv/TP_VehicleAdv.cpp
--- a/Templates/TP_VehicleAdv/Source/TP_VehicleAdv/TP_VehicleAdv.cpp
+++ b/Templates/TP_VehicleAdv/Source/TP_VehicleAdv/TP_VehicleAdv.cpp
@@ -1,4 +1,4 @@
-// Copyright 1998-2018 Epic Games, Inc. All Rights Reserved.
+// Copyright 1998-2019 Epic Games, Inc. All Rights Reserved.
 
 #include "TP_VehicleAdv.h"
 #include "Modules/ModuleManager.h"
diff --git a/Templates/TP_VehicleAdv/Source/TP_VehicleAdv/TP_VehicleAdv.h b/Templates/TP_VehicleAdv/Source/TP_VehicleAdv/TP_VehicleAdv.h
--- a/Templates/TP_VehicleAdv/Source/TP_VehicleAdv/TP_VehicleAdv.h
+++ b/Templates/TP_VehicleAdv/Source/TP_VehicleAdv/TP_VehicleAdv.h
@@ -1,4 +1,4 @@
-// Copyright 1998-2018 Epic Games, Inc. All Rights Reserved.
+// Copyright 1998-2019 Epic Games, Inc. All Rights Reserved.
 
 #pragma once
 
diff --git a/Templates/TP_VehicleAdv/Source/TP_VehicleAdv/TP_VehicleAdvGameMode.cpp b/Templates/TP_VehicleAdv/Source/TP_VehicleAdv/TP_VehicleAdvGameMode.cpp
--- a/Templates/TP_VehicleAdv/Source/TP_VehicleAdv/TP_VehicleAdvGameMode.cpp
+++ b/Templates/TP_VehicleAdv/Source/TP_VehicleAdv/TP_VehicleAdvGameMode.cpp
@@ -1,4 +1,4 @@
-// Copyright 1998-2018 Epic Games, Inc. All Rights Reserved.
+// Copyright 1998-2019 Epic Games, Inc. All Rights Reserved.
 
 #include "TP_VehicleAdvGameMode.h"
 #include "TP_VehicleAdvPawn.h"
diff --git a/Templates/TP_VehicleAdv/Source/TP_VehicleAdv/TP_VehicleAdvGameMode.h b/Templates/TP_VehicleAdv/Source/TP_VehicleAdv/TP_VehicleAdvGameMode.h
--- a/Templates/TP_VehicleAdv/Source/TP_VehicleAdv/TP_VehicleAdvGameMode.h
+++ b/Templates/TP_VehicleAdv/Source/TP_VehicleAdv/TP_VehicleAdvGameMode.h
@@ -1,4 +1,4 @@
-// Copyright 1998-2018 Epic Games, Inc. All Rights Reserved.
+// Copyright 1998-2019 Epic Games, Inc. All Rights Reserved.
 
 #pragma once
 
diff --git a/Templates/TP_VehicleAdv/Source/TP_VehicleAdv/TP_VehicleAdvHud.cpp b/Templates/TP_VehicleAdv/Source/TP_VehicleAdv/TP_VehicleAdvHud.cpp
--- a/Templates/TP_VehicleAdv/Source/TP_VehicleAdv/TP_VehicleAdvHud.cpp
+++ b/Templates/TP_VehicleAdv/Source/TP_VehicleAdv/TP_VehicleAdvHud.cpp
@@ -1,4 +1,4 @@
-// Copyright 1998-2018 Epic Games, Inc. All Rights Reserved.
+// Copyright 1998-2019 Epic Games, Inc. All Rights Reserved.
 
 #include "TP_VehicleAdvHud.h"
 #include "TP_VehicleAdvPawn.h"
diff --git a/Templates/TP_VehicleAdv/Source/TP_VehicleAdv/TP_VehicleAdvHud.h b/Templates/TP_VehicleAdv/Source/TP_VehicleAdv/TP_VehicleAdvHud.h
--- a/Templates/TP_VehicleAdv/Source/TP_VehicleAdv/TP_VehicleAdvHud.h
+++ b/Templates/TP_VehicleAdv/Source/TP_VehicleAdv/TP_VehicleAdvHud.h
@@ -1,4 +1,4 @@
-// Copyright 1998-2018 Epic Games, Inc. All Rights Reserved.
+// Copyright 1998-2019 Epic Games, Inc. All Rights Reserved.
 
 #pragma once
 
diff --git a/Templates/TP_VehicleAdv/Source/TP_VehicleAdv/TP_VehicleAdvPawn.cpp b/Templates/TP_VehicleAdv/Source/TP_VehicleAdv/TP_VehicleAdvPawn.cpp
--- a/Templates/TP_VehicleAdv/Source/TP_VehicleAdv/TP_VehicleAdvPawn.cpp
+++ b/Templates/TP_VehicleAdv/Source/TP_VehicleAdv/TP_VehicleAdvPawn.cpp
@@ -1,4 +1,4 @@
-// Copyright 1998-2018 Epic Games, Inc. All Rights Reserved.
+// Copyright 1998-2019 Epic Games, Inc. All Rights Reserved.
 
 #include "TP_VehicleAdvPawn.h"
 #include "TP_VehicleAdvWheelFront.h"
diff --git a/Templates/TP_VehicleAdv/Source/TP_VehicleAdv/TP_VehicleAdvPawn.h b/Templates/TP_VehicleAdv/Source/TP_VehicleAdv/TP_VehicleAdvPawn.h
--- a/Templates/TP_VehicleAdv/Source/TP_VehicleAdv/TP_VehicleAdvPawn.h
+++ b/Templates/TP_VehicleAdv/Source/TP_VehicleAdv/TP_VehicleAdvPawn.h
@@ -1,4 +1,4 @@
-// Copyright 1998-2018 Epic Games, Inc. All Rights Reserved.
+// Copyright 1998-2019 Epic Games, Inc. All Rights Reserved.
 
 #pragma once
 
diff --git a/Templates/TP_VehicleAdv/Source/TP_VehicleAdv/TP_VehicleAdvWheelFront.cpp b/Templates/TP_VehicleAdv/Source/TP_VehicleAdv/TP_VehicleAdvWheelFront.cpp
--- a/Templates/TP_VehicleAdv/Source/TP_VehicleAdv/TP_VehicleAdvWheelFront.cpp
+++ b/Templates/TP_VehicleAdv/Source/TP_VehicleAdv/TP_VehicleAdvWheelFront.cpp
@@ -1,4 +1,4 @@
-// Copyright 1998-2018 Epic Games, Inc. All Rights Reserved.
+// Copyright 1998-2019 Epic Games, Inc. All Rights Reserved.
 
 #include "TP_VehicleAdvWheelFront.h"
 #include "TireConfig.h"
diff --git a/Templates/TP_VehicleAdv/Source/TP_VehicleAdv/TP_VehicleAdvWheelFront.h b/Templates/TP_VehicleAdv/Source/TP_VehicleAdv/TP_VehicleAdvWheelFront.h
--- a/Templates/TP_VehicleAdv/Source/TP_VehicleAdv/TP_VehicleAdvWheelFront.h
+++ b/Templates/TP_VehicleAdv/Source/TP_VehicleAdv/TP_VehicleAdvWheelFront.h
@@ -1,4 +1,4 @@
-// Copyright 1998-2018 Epic Games, Inc. All Rights Reserved.
+// Copyright 1998-2019 Epic Games, Inc. All Rights Reserved.
 
 #pragma once
 
diff --git a/Templates/TP_VehicleAdv/Source/TP_VehicleAdv/TP_VehicleAdvWheelRear.cpp b/Templates/TP_VehicleAdv/Source/TP_VehicleAdv/TP_VehicleAdvWheelRear.cpp
--- a/Templates/TP_VehicleAdv/Source/TP_VehicleAdv/TP_VehicleAdvWheelRear.cpp
+++ b/Templates/TP_VehicleAdv/Source/TP_VehicleAdv/TP_VehicleAdvWheelRear.cpp
@@ -1,4 +1,4 @@
-// Copyright 1998-2018 Epic Games, Inc. All Rights Reserved.
+// Copyright 1998-2019 Epic Games, Inc. All Rights Reserved.
 
 #include "TP_VehicleAdvWheelRear.h"
 #include "TireConfig.h"
diff --git a/Templates/TP_VehicleAdv/Source/TP_VehicleAdv/TP_VehicleAdvWheelRear.h b/Templates/TP_VehicleAdv/Source/TP_VehicleAdv/TP_VehicleAdvWheelRear.h
--- a/Templates/TP_VehicleAdv/Source/TP_VehicleAdv/TP_VehicleAdvWheelRear.h
+++ b/Templates/TP_VehicleAdv/Source/TP_VehicleAdv/TP_VehicleAdvWheelRear.h
@@ -1,4 +1,4 @@
-// Copyright 1998-2018 Epic Games, Inc. All Rights Reserved.
+// Copyright 1998-2019 Epic Games, Inc. All Rights Reserved.
 
 #pragma once
 
diff --git a/Templates/TP_VehicleAdv/Source/TP_VehicleAdvEditor.Target.cs b/Templates/TP_VehicleAdv/Source/TP_VehicleAdvEditor.Target.cs
--- a/Templates/TP_VehicleAdv/Source/TP_VehicleAdvEditor.Target.cs
+++ b/Templates/TP_VehicleAdv/Source/TP_VehicleAdvEditor.Target.cs
@@ -1,4 +1,4 @@
-// Copyright 1998-2018 Epic Games, Inc. All Rights Reserved.
+// Copyright 1998-2019 Epic Games, Inc. All Rights Reserved.
 
 using UnrealBuildTool;
 using System.Collections.Generic;
